Consolidate cinematics items in Window menu

diff --git a/Engine/Plugins/VirtualProduction/Takes/Source/TakeRecorder/Private/TakeRecorderModule.cpp b/Engine/Plugins/VirtualProduction/Takes/Source/TakeRecorder/Private/TakeRecorderModule.cpp
--- a/Engine/Plugins/VirtualProduction/Takes/Source/TakeRecorder/Private/TakeRecorderModule.cpp
+++ b/Engine/Plugins/VirtualProduction/Takes/Source/TakeRecorder/Private/TakeRecorderModule.cpp
@@ -141,7 +141,7 @@ namespace
 		FTabSpawnerEntry& TabSpawner = FGlobalTabmanager::Get()->RegisterNomadTabSpawner(ITakeRecorderModule::TakeRecorderTabName, FOnSpawnTab::CreateStatic(SpawnTakeRecorderTab));
 
 		TabSpawner
-			.SetGroup(WorkspaceMenu::GetMenuStructure().GetLevelEditorCategory())
+			.SetGroup(WorkspaceMenu::GetMenuStructure().GetLevelEditorCinematicsCategory())
 			.SetDisplayName(ITakeRecorderModule::TakeRecorderTabLabel)
 			.SetTooltipText(LOCTEXT("TakeRecorderTab_Tooltip", "Open the main Take Recorder UI."))
 			.SetIcon(FSlateIcon(FTakeRecorderStyle::StyleName, "TakeRecorder.TabIcon"));
@@ -149,7 +149,7 @@ namespace
 		FTabSpawnerEntry& TBTabSpawner = FGlobalTabmanager::Get()->RegisterNomadTabSpawner(ITakeRecorderModule::TakesBrowserTabName, FOnSpawnTab::CreateStatic(SpawnTakesBrowserTab));
 
 		TBTabSpawner
-			.SetGroup(WorkspaceMenu::GetMenuStructure().GetLevelEditorCategory())
+			.SetGroup(WorkspaceMenu::GetMenuStructure().GetLevelEditorCinematicsCategory())
 			.SetDisplayName(ITakeRecorderModule::TakesBrowserTabLabel)
 			.SetTooltipText(LOCTEXT("TakeBrowserTab_Tooltip", "Open the Take Browser UI"))
 			.SetIcon(FSlateIcon(FEditorStyle::GetStyleSetName(), "ContentBrowser.TabIcon"));
diff --git a/Engine/Source/Editor/EditorStyle/Private/SlateEditorStyle.cpp b/Engine/Source/Editor/EditorStyle/Private/SlateEditorStyle.cpp
--- a/Engine/Source/Editor/EditorStyle/Private/SlateEditorStyle.cpp
+++ b/Engine/Source/Editor/EditorStyle/Private/SlateEditorStyle.cpp
@@ -4809,6 +4809,7 @@ void FSlateEditorStyle::FStyle::SetupLevelEditorStyle()
 #if WITH_EDITOR || (IS_PROGRAM && WITH_UNREAL_DEVELOPER_TOOLS)
 		{
 			Set( "LevelEditor.Tabs.Details", new IMAGE_BRUSH( "/Icons/icon_tab_SelectionDetails_16x", Icon16x16 ) );
+			Set( "LevelEditor.Tabs.Cinematics", new IMAGE_BRUSH("/Icons/icon_tab_Cinematics_16x", Icon16x16));
 			Set( "LevelEditor.Tabs.EditorModes", new IMAGE_BRUSH( "/Icons/icon_Editor_Modes_16x", Icon16x16 ) );
 			Set( "LevelEditor.Tabs.Modes", new IMAGE_BRUSH( "/Icons/icon_Editor_Modes_16x", Icon16x16 ) );
 			Set( "LevelEditor.Tabs.Properties", new IMAGE_BRUSH( "/Icons/properties_16x", Icon16x16 ) );
diff --git a/Engine/Source/Editor/LevelEditor/Private/SLevelEditor.cpp b/Engine/Source/Editor/LevelEditor/Private/SLevelEditor.cpp
--- a/Engine/Source/Editor/LevelEditor/Private/SLevelEditor.cpp
+++ b/Engine/Source/Editor/LevelEditor/Private/SLevelEditor.cpp
@@ -1198,7 +1198,7 @@ TSharedRef<SWidget> SLevelEditor::RestoreContentArea( const TSharedRef<SDockTab>
 			const FSlateIcon SequencerIcon("LevelSequenceEditorStyle", "LevelSequenceEditor.Tabs.Sequencer" );
 			LevelEditorTabManager->RegisterTabSpawner( "Sequencer", FOnSpawnTab::CreateSP<SLevelEditor, FName, FString>(this, &SLevelEditor::SpawnLevelEditorTab, FName("Sequencer"), FString()) )
 				.SetDisplayName(NSLOCTEXT("LevelEditorTabs", "Sequencer", "Sequencer"))
-				.SetGroup( MenuStructure.GetLevelEditorCategory() )
+				.SetGroup( MenuStructure.GetLevelEditorCinematicsCategory() )
 				.SetIcon( SequencerIcon );
 		}
 
diff --git a/Engine/Source/Editor/SequenceRecorder/Private/SequenceRecorderModule.cpp b/Engine/Source/Editor/SequenceRecorder/Private/SequenceRecorderModule.cpp
--- a/Engine/Source/Editor/SequenceRecorder/Private/SequenceRecorderModule.cpp
+++ b/Engine/Source/Editor/SequenceRecorder/Private/SequenceRecorderModule.cpp
@@ -157,7 +157,7 @@ class FSequenceRecorderModule : public ISequenceRecorder, private FSelfRegisteri
 			auto RegisterTabSpawner = []()
 			{
 				FGlobalTabmanager::Get()->RegisterNomadTabSpawner(SequenceRecorderTabName, FOnSpawnTab::CreateStatic(&FSequenceRecorderModule::SpawnSequenceRecorderTab))
-				.SetGroup(WorkspaceMenu::GetMenuStructure().GetLevelEditorCategory())
+				.SetGroup(WorkspaceMenu::GetMenuStructure().GetLevelEditorCinematicsCategory())
 				.SetDisplayName(LOCTEXT("SequenceRecorderTabTitle", "Sequence Recorder"))
 				.SetTooltipText(LOCTEXT("SequenceRecorderTooltipText", "Open the Sequence Recorder tab."))
 				.SetIcon(FSlateIcon(FEditorStyle::GetStyleSetName(), "SequenceRecorder.TabIcon"));
diff --git a/Engine/Source/Editor/WorkspaceMenuStructure/Private/WorkspaceMenuStructureModule.cpp b/Engine/Source/Editor/WorkspaceMenuStructure/Private/WorkspaceMenuStructureModule.cpp
--- a/Engine/Source/Editor/WorkspaceMenuStructure/Private/WorkspaceMenuStructureModule.cpp
+++ b/Engine/Source/Editor/WorkspaceMenuStructure/Private/WorkspaceMenuStructureModule.cpp
@@ -41,6 +41,11 @@ public:
 		return LevelEditorModesCategory.ToSharedRef();
 	}
 
+	virtual TSharedRef<FWorkspaceItem> GetLevelEditorCinematicsCategory() const override
+	{
+		return LevelEditorCinematicsCategory.ToSharedRef();
+	}
+
 	virtual TSharedRef<FWorkspaceItem> GetToolsCategory() const override
 	{
 		return ToolsCategory.ToSharedRef();
@@ -76,6 +81,7 @@ public:
 		LevelEditorCategory->ClearItems();
 		LevelEditorViewportsCategory = LevelEditorCategory->AddGroup(LOCTEXT( "WorkspaceMenu_LevelEditorViewportCategory", "Viewports" ), LOCTEXT( "WorkspaceMenu_LevelEditorViewportCategoryTooltip", "Open a Viewport tab." ), FSlateIcon(FEditorStyle::GetStyleSetName(), "LevelEditor.Tabs.Viewports"), true);
 		LevelEditorDetailsCategory = LevelEditorCategory->AddGroup(LOCTEXT("WorkspaceMenu_LevelEditorDetailCategory", "Details" ), LOCTEXT("WorkspaceMenu_LevelEditorDetailCategoryTooltip", "Open a Details tab." ), FSlateIcon(FEditorStyle::GetStyleSetName(), "LevelEditor.Tabs.Details"), true );
+		LevelEditorCinematicsCategory = LevelEditorCategory->AddGroup(LOCTEXT("WorkspaceMenu_LevelEditorCinematicsCategory", "Cinematics"), LOCTEXT("WorkspaceMenu_LevelEditorCinematicsCategoryTooltip", "Open a Cinematics tab."), FSlateIcon(FEditorStyle::GetStyleSetName(), "LevelEditor.Tabs.Cinematics"), true);
 		LevelEditorModesCategory = LevelEditorCategory->AddGroup(LOCTEXT("WorkspaceMenu_LevelEditorToolsCategory", "Editor Modes" ), FSlateIcon(FEditorStyle::GetStyleSetName(), "LevelEditor.Tabs.EditorModes"), true );
 	}
 
@@ -114,6 +120,7 @@ private:
 	TSharedPtr<FWorkspaceItem> LevelEditorCategory;
 	TSharedPtr<FWorkspaceItem> LevelEditorViewportsCategory;
 	TSharedPtr<FWorkspaceItem> LevelEditorDetailsCategory;
+	TSharedPtr<FWorkspaceItem> LevelEditorCinematicsCategory;
 	TSharedPtr<FWorkspaceItem> LevelEditorModesCategory;
 
 	TSharedPtr<FWorkspaceItem> ToolsCategory;
diff --git a/Engine/Source/Editor/WorkspaceMenuStructure/Public/WorkspaceMenuStructure.h b/Engine/Source/Editor/WorkspaceMenuStructure/Public/WorkspaceMenuStructure.h
--- a/Engine/Source/Editor/WorkspaceMenuStructure/Public/WorkspaceMenuStructure.h
+++ b/Engine/Source/Editor/WorkspaceMenuStructure/Public/WorkspaceMenuStructure.h
@@ -40,6 +40,9 @@ public:
 	/** See diagram above */
 	virtual TSharedRef<FWorkspaceItem> GetLevelEditorDetailsCategory() const = 0;
 
+	/** See diagram above */
+	virtual TSharedRef<FWorkspaceItem> GetLevelEditorCinematicsCategory() const = 0;
+
 	/** See diagram above */
 	virtual TSharedRef<FWorkspaceItem> GetLevelEditorModesCategory() const = 0;
 
