Extend Workspace Menus With Stable Section Names

Note: Workspace Menus is the name of the menu hierarchy created by tab spawners.

The intent here is to give stable section and sub-menu names to the main menu hierarchy created by tab spawners to allow anyone to extend said hierarchy using UToolMenus.

* Add a NameAttribute to FWorkspaceItem which is NAME_None by default. Allow NameAttribute to be set through constructors and the NewGroup and AddGroup methods. Add an FWorkspaceItem::GetFName method for fetching NameAttribute.
* Set the NameAttribute of FWorkspaceItems created for tab spawners using their TabType.
* When adding tab spawners to the main menus, set the FName identifier as appropriate on sections and sub-menus with the value returned by FWorkspaceItem::GetFName.
* Set the NameAttribute on a lot of workspace groups to make use of this new functionality.

diff --git a/Engine/Plugins/Media/ImgMedia/Source/ImgMediaEditor/Private/ImgMediaEditorModule.cpp b/Engine/Plugins/Media/ImgMedia/Source/ImgMediaEditor/Private/ImgMediaEditorModule.cpp
--- a/Engine/Plugins/Media/ImgMedia/Source/ImgMediaEditor/Private/ImgMediaEditorModule.cpp
+++ b/Engine/Plugins/Media/ImgMedia/Source/ImgMediaEditor/Private/ImgMediaEditorModule.cpp
@@ -123,6 +123,7 @@ protected:
 		// Add ImgMedia group.
 		const IWorkspaceMenuStructure& MenuStructure = WorkspaceMenu::GetMenuStructure();
 		TSharedRef<FWorkspaceItem> MediaBrowserGroup = MenuStructure.GetLevelEditorCategory()->AddGroup(
+			"ImgMedia",
 			LOCTEXT("WorkspaceMenu_ImgMediaCategory", "ImgMedia"),
 			FSlateIcon(FAppStyle::GetAppStyleSetName(), "SequenceRecorder.TabIcon"),
 			true);
diff --git a/Engine/Source/Editor/ContentBrowser/Private/ContentBrowserSingleton.cpp b/Engine/Source/Editor/ContentBrowser/Private/ContentBrowserSingleton.cpp
--- a/Engine/Source/Editor/ContentBrowser/Private/ContentBrowserSingleton.cpp
+++ b/Engine/Source/Editor/ContentBrowser/Private/ContentBrowserSingleton.cpp
@@ -92,6 +92,7 @@ FContentBrowserSingleton::FContentBrowserSingleton()
 	const FSlateIcon ContentBrowserIcon(FAppStyle::Get().GetStyleSetName(), "ContentBrowser.TabIcon");
 	const IWorkspaceMenuStructure& MenuStructure = WorkspaceMenu::GetMenuStructure();
 	TSharedRef<FWorkspaceItem> ContentBrowserGroup = MenuStructure.GetLevelEditorCategory()->AddGroup(
+		"ContentBrowser",
 		LOCTEXT("WorkspaceMenu_ContentBrowserCategory", "Content Browser"),
 		LOCTEXT("ContentBrowserMenuTooltipText", "Open a Content Browser tab."),
 		ContentBrowserIcon,
diff --git a/Engine/Source/Editor/Kismet/Private/FindInBlueprintManager.cpp b/Engine/Source/Editor/Kismet/Private/FindInBlueprintManager.cpp
--- a/Engine/Source/Editor/Kismet/Private/FindInBlueprintManager.cpp
+++ b/Engine/Source/Editor/Kismet/Private/FindInBlueprintManager.cpp
@@ -3683,6 +3683,7 @@ void FFindInBlueprintSearchManager::EnableGlobalFindResults(bool bEnable)
 		// Register the spawners for all global Find Results tabs
 		const FSlateIcon GlobalFindResultsIcon(FAppStyle::GetAppStyleSetName(), "BlueprintEditor.FindInBlueprints.MenuIcon");
 		GlobalFindResultsMenuItem = WorkspaceMenu::GetMenuStructure().GetToolsCategory()->AddGroup(
+			"FindInBlueprints",
 			LOCTEXT("WorkspaceMenu_GlobalFindResultsCategory", "Find in Blueprints"),
 			LOCTEXT("GlobalFindResultsMenuTooltipText", "Find references to functions, events and variables in all Blueprints."),
 			GlobalFindResultsIcon,
diff --git a/Engine/Source/Editor/MainFrame/Private/Menus/MainMenu.cpp b/Engine/Source/Editor/MainFrame/Private/Menus/MainMenu.cpp
--- a/Engine/Source/Editor/MainFrame/Private/Menus/MainMenu.cpp
+++ b/Engine/Source/Editor/MainFrame/Private/Menus/MainMenu.cpp
@@ -55,17 +55,10 @@ void MakeSpawnerToolMenuEntry(TSharedPtr<FTabManager> TabManager, FToolMenuSecti
 
 void PopulateTabSpawnerToolMenu_Helper(UToolMenu* InMenu, TSharedPtr<FTabManager> TabManager, TSharedRef<FWorkspaceItem> InMenuStructure, TSharedRef< TArray< TWeakPtr<FTabSpawnerEntry> > > AllSpawners, const int32 RecursionLevel, bool PutLabelOnSection)
 {
-	FToolMenuSection* TheSection = nullptr;
+	FToolMenuSection& Section = InMenu->FindOrAddSection(InMenuStructure->GetFName());
 	if (PutLabelOnSection)
 	{
-		const FText Label = InMenuStructure->GetDisplayName();
-		const FName Name = MakeObjectNameFromDisplayLabel(Label.ToString(), FName(TEXT("SectionName")));
-		TheSection = &InMenu->FindOrAddSection(Name);
-		TheSection->Label = Label;
-	}
-	else
-	{
-		TheSection = &InMenu->FindOrAddSection(NAME_None);
+		Section.Label = InMenuStructure->GetDisplayName();
 	}
 
 	for (const TSharedRef<FWorkspaceItem>& Child : InMenuStructure->GetChildItems())
@@ -77,7 +70,7 @@ void PopulateTabSpawnerToolMenu_Helper(UToolMenu* InMenu, TSharedPtr<FTabManager
 			// Only show non-hidden items that have a valid spawner.
 			if (AllSpawners->Contains(Spawner.ToSharedRef()))
 			{
-				MakeSpawnerToolMenuEntry(TabManager, *TheSection, Spawner);
+				MakeSpawnerToolMenuEntry(TabManager, Section, Spawner);
 			}
 		}
 		else if (Child->HasChildrenIn(*AllSpawners))
@@ -91,8 +84,8 @@ void PopulateTabSpawnerToolMenu_Helper(UToolMenu* InMenu, TSharedPtr<FTabManager
 			else
 			{
 				// Create a submenu and add all children there.
-				FToolMenuEntry& SubMenu = TheSection->AddSubMenu(
-					NAME_None,
+				FToolMenuEntry& SubMenu = Section.AddSubMenu(
+					Child->GetFName(),
 					Child->GetDisplayName(),
 					Child->GetTooltipText(),
 					FNewToolMenuDelegate::CreateStatic(&PopulateTabSpawnerToolMenu_Helper, TabManager, Child, AllSpawners, RecursionLevel+1, false),
diff --git a/Engine/Source/Editor/WorkspaceMenuStructure/Private/WorkspaceMenuStructureModule.cpp b/Engine/Source/Editor/WorkspaceMenuStructure/Private/WorkspaceMenuStructureModule.cpp
--- a/Engine/Source/Editor/WorkspaceMenuStructure/Private/WorkspaceMenuStructureModule.cpp
+++ b/Engine/Source/Editor/WorkspaceMenuStructure/Private/WorkspaceMenuStructureModule.cpp
@@ -115,28 +115,28 @@ public:
 	void ResetLevelEditorCategory()
 	{
 		LevelEditorCategory->ClearItems();
-		LevelEditorViewportsCategory = LevelEditorCategory->AddGroup(LOCTEXT( "WorkspaceMenu_LevelEditorViewportCategory", "Viewports" ), LOCTEXT( "WorkspaceMenu_LevelEditorViewportCategoryTooltip", "Open a Viewport tab." ), FSlateIcon(FAppStyle::Get().GetStyleSetName(), "LevelEditor.Tabs.Viewports"), true);
-		LevelEditorDetailsCategory = LevelEditorCategory->AddGroup(LOCTEXT("WorkspaceMenu_LevelEditorDetailCategory", "Details" ), LOCTEXT("WorkspaceMenu_LevelEditorDetailCategoryTooltip", "Open a Details tab." ), FSlateIcon(FAppStyle::Get().GetStyleSetName(), "LevelEditor.Tabs.Details"), true );
-		LevelEditorCinematicsCategory = LevelEditorCategory->AddGroup(LOCTEXT("WorkspaceMenu_LevelEditorCinematicsCategory", "Cinematics"), LOCTEXT("WorkspaceMenu_LevelEditorCinematicsCategoryTooltip", "Open a Cinematics tab."), FSlateIcon(FAppStyle::Get().GetStyleSetName(), "LevelEditor.Tabs.Cinematics"), true);
-		LevelEditorVirtualProductionCategory = LevelEditorCategory->AddGroup(LOCTEXT("WorkspaceMenu_LevelEditorVirtualProductionCategory", "Virtual Production"), LOCTEXT("WorkspaceMenu_LevelEditorVirtualProductionCategoryTooltip", "Open a Virtual Production tab."), FSlateIcon(FAppStyle::Get().GetStyleSetName(), "LevelEditor.Tabs.VirtualProduction"), true);
-		LevelEditorModesCategory = LevelEditorCategory->AddGroup(LOCTEXT("WorkspaceMenu_LevelEditorToolsCategory", "Editor Modes" ), FSlateIcon(FAppStyle::Get().GetStyleSetName(), "LevelEditor.Tabs.EditorModes"), true );
-		LevelEditorWorldPartitionCategory = LevelEditorCategory->AddGroup(LOCTEXT("WorkspaceMenu_LevelEditorWorldPartitionCategory", "World Partition"), LOCTEXT("WorkspaceMenu_LevelEditorWorldPartitionCategoryTooltip", "Open a World Partition tab."), FSlateIcon(FAppStyle::Get().GetStyleSetName(), "LevelEditor.Tabs.WorldPartition"), true);
-		LevelEditorOutlinerCategory = LevelEditorCategory->AddGroup(LOCTEXT("WorkspaceMenu_LevelEditorOutlinerCategory", "Outliner"), LOCTEXT("WorkspaceMenu_LevelEditorOutlinerCategoryTooltip", "Open an Outliner tab."), FSlateIcon(FAppStyle::GetAppStyleSetName(), "LevelEditor.Tabs.Outliner"), true);
+		LevelEditorViewportsCategory = LevelEditorCategory->AddGroup("LevelEditorViewport", LOCTEXT("WorkspaceMenu_LevelEditorViewportCategory", "Viewports"), LOCTEXT( "WorkspaceMenu_LevelEditorViewportCategoryTooltip", "Open a Viewport tab." ), FSlateIcon(FAppStyle::Get().GetStyleSetName(), "LevelEditor.Tabs.Viewports"), true);
+		LevelEditorDetailsCategory = LevelEditorCategory->AddGroup("LevelEditorDetail", LOCTEXT("WorkspaceMenu_LevelEditorDetailCategory", "Details" ), LOCTEXT("WorkspaceMenu_LevelEditorDetailCategoryTooltip", "Open a Details tab." ), FSlateIcon(FAppStyle::Get().GetStyleSetName(), "LevelEditor.Tabs.Details"), true );
+		LevelEditorCinematicsCategory = LevelEditorCategory->AddGroup("LevelEditorCinematics", LOCTEXT("WorkspaceMenu_LevelEditorCinematicsCategory", "Cinematics"), LOCTEXT("WorkspaceMenu_LevelEditorCinematicsCategoryTooltip", "Open a Cinematics tab."), FSlateIcon(FAppStyle::Get().GetStyleSetName(), "LevelEditor.Tabs.Cinematics"), true);
+		LevelEditorVirtualProductionCategory = LevelEditorCategory->AddGroup("LevelEditorVirtualProduction", LOCTEXT("WorkspaceMenu_LevelEditorVirtualProductionCategory", "Virtual Production"), LOCTEXT("WorkspaceMenu_LevelEditorVirtualProductionCategoryTooltip", "Open a Virtual Production tab."), FSlateIcon(FAppStyle::Get().GetStyleSetName(), "LevelEditor.Tabs.VirtualProduction"), true);
+		LevelEditorModesCategory = LevelEditorCategory->AddGroup("LevelEditorTools", LOCTEXT("WorkspaceMenu_LevelEditorToolsCategory", "Editor Modes" ), FSlateIcon(FAppStyle::Get().GetStyleSetName(), "LevelEditor.Tabs.EditorModes"), true );
+		LevelEditorWorldPartitionCategory = LevelEditorCategory->AddGroup("LevelEditorWorldPartition", LOCTEXT("WorkspaceMenu_LevelEditorWorldPartitionCategory", "World Partition"), LOCTEXT("WorkspaceMenu_LevelEditorWorldPartitionCategoryTooltip", "Open a World Partition tab."), FSlateIcon(FAppStyle::Get().GetStyleSetName(), "LevelEditor.Tabs.WorldPartition"), true);
+		LevelEditorOutlinerCategory = LevelEditorCategory->AddGroup("LevelEditorOutliner", LOCTEXT("WorkspaceMenu_LevelEditorOutlinerCategory", "Outliner"), LOCTEXT("WorkspaceMenu_LevelEditorOutlinerCategoryTooltip", "Open an Outliner tab."), FSlateIcon(FAppStyle::GetAppStyleSetName(), "LevelEditor.Tabs.Outliner"), true);
 	}
 
 	void ResetToolsCategory()
 	{
 		// Developer tools sections
 		DeveloperToolsCategory->ClearItems();
-		DeveloperToolsSubMenuDebugCategory =     DeveloperToolsCategory->AddGroup(LOCTEXT("WorkspaceMenu_DeveloperToolsDebugCategory", "Debug"), FSlateIcon(FAppStyle::Get().GetStyleSetName(), "LevelEditor.Tabs.Debug"), true);
+		DeveloperToolsSubMenuDebugCategory =     DeveloperToolsCategory->AddGroup("DeveloperToolsDebug", LOCTEXT("WorkspaceMenu_DeveloperToolsDebugCategory", "Debug"), FSlateIcon(FAppStyle::Get().GetStyleSetName(), "LevelEditor.Tabs.Debug"), true);
 
-		DeveloperToolsDebugCategory = DeveloperToolsSubMenuDebugCategory->AddGroup(LOCTEXT("WorkspaceMenu_DeveloperToolsDebugCategory", "Debug"), FSlateIcon(), true);
-		AutomationToolsCategory = DeveloperToolsSubMenuDebugCategory->AddGroup(LOCTEXT("WorkspaceMenu_AutomationToolsCategory", "Testing"), FSlateIcon(), true);
+		DeveloperToolsDebugCategory = DeveloperToolsSubMenuDebugCategory->AddGroup("DebugTools", LOCTEXT("WorkspaceMenu_DeveloperToolsDebugCategory", "Debug"), FSlateIcon(), true);
+		AutomationToolsCategory = DeveloperToolsSubMenuDebugCategory->AddGroup("TestingTools", LOCTEXT("WorkspaceMenu_AutomationToolsCategory", "Testing"), FSlateIcon(), true);
 
-		DeveloperToolsProfilingCategory = DeveloperToolsCategory->AddGroup(LOCTEXT("WorkspaceMenu_DeveloperToolsProfilingCategory", "Profile"),  FSlateIcon(FAppStyle::Get().GetStyleSetName(), "LevelEditor.Profile"), true);
-		DeveloperToolsAuditCategory =     DeveloperToolsCategory->AddGroup(LOCTEXT("WorkspaceMenu_DeveloperToolsAuditCategory", "Audit"), FSlateIcon(FAppStyle::Get().GetStyleSetName(), "LevelEditor.Audit"), true);
-		DeveloperToolsPlatformsCategory = DeveloperToolsCategory->AddGroup(LOCTEXT("WorkspaceMenu_DeveloperToolsPlatforms:WCategory", "Platforms"), FSlateIcon(FAppStyle::Get().GetStyleSetName(), "LevelEditor.Platforms"), true);
-		DeveloperToolsMiscCategory =      DeveloperToolsCategory->AddGroup(LOCTEXT("WorkspaceMenu_DeveloperToolsMiscCategory", "Miscellaneous"), FSlateIcon(FAppStyle::Get().GetStyleSetName(), "LevelEditor.Tabs.Debug"), true);
+		DeveloperToolsProfilingCategory = DeveloperToolsCategory->AddGroup("DeveloperToolsProfiling", LOCTEXT("WorkspaceMenu_DeveloperToolsProfilingCategory", "Profile"),  FSlateIcon(FAppStyle::Get().GetStyleSetName(), "LevelEditor.Profile"), true);
+		DeveloperToolsAuditCategory =     DeveloperToolsCategory->AddGroup("DeveloperToolsAudit", LOCTEXT("WorkspaceMenu_DeveloperToolsAuditCategory", "Audit"), FSlateIcon(FAppStyle::Get().GetStyleSetName(), "LevelEditor.Audit"), true);
+		DeveloperToolsPlatformsCategory = DeveloperToolsCategory->AddGroup("DeveloperToolsPlatforms", LOCTEXT("WorkspaceMenu_DeveloperToolsPlatforms:WCategory", "Platforms"), FSlateIcon(FAppStyle::Get().GetStyleSetName(), "LevelEditor.Platforms"), true);
+		DeveloperToolsMiscCategory =      DeveloperToolsCategory->AddGroup("DeveloperToolsMisc", LOCTEXT("WorkspaceMenu_DeveloperToolsMiscCategory", "Miscellaneous"), FSlateIcon(FAppStyle::Get().GetStyleSetName(), "LevelEditor.Tabs.Debug"), true);
 	}
 
 public:
diff --git a/Engine/Source/Runtime/Slate/Public/Framework/Docking/TabManager.h b/Engine/Source/Runtime/Slate/Public/Framework/Docking/TabManager.h
--- a/Engine/Source/Runtime/Slate/Public/Framework/Docking/TabManager.h
+++ b/Engine/Source/Runtime/Slate/Public/Framework/Docking/TabManager.h
@@ -215,7 +215,7 @@ DECLARE_MULTICAST_DELEGATE_OneParam(FOnReadOnlyModeChanged, bool /*bReadOnly*/);
 struct FTabSpawnerEntry : public FWorkspaceItem
 {
 	FTabSpawnerEntry(const FName& InTabType, const FOnSpawnTab& InSpawnTabMethod, const FCanSpawnTab& InCanSpawnTab)
-		: FWorkspaceItem(FText(), FSlateIcon(), false)
+		: FWorkspaceItem(InTabType, FText(), FSlateIcon(), false)
 		, TabType(InTabType)
 		, OnSpawnTab(InSpawnTabMethod)
 		, CanSpawnTab(InCanSpawnTab)
diff --git a/Engine/Source/Runtime/Slate/Public/Framework/Docking/WorkspaceItem.h b/Engine/Source/Runtime/Slate/Public/Framework/Docking/WorkspaceItem.h
--- a/Engine/Source/Runtime/Slate/Public/Framework/Docking/WorkspaceItem.h
+++ b/Engine/Source/Runtime/Slate/Public/Framework/Docking/WorkspaceItem.h
@@ -34,28 +34,57 @@ protected:
 public:
 	static TSharedRef<FWorkspaceItem> NewGroup( const FText& DisplayName, const FSlateIcon& Icon = FSlateIcon(), const bool bSortChildren = false )
 	{
-		return MakeShareable( new FWorkspaceItem( DisplayName, Icon, bSortChildren ) );
+		return MakeShareable( new FWorkspaceItem( NAME_None, DisplayName, Icon, bSortChildren ) );
+	}
+
+	static TSharedRef<FWorkspaceItem> NewGroup( const FName& Name, const FText& DisplayName, const FSlateIcon& Icon = FSlateIcon(), const bool bSortChildren = false )
+	{
+		return MakeShareable( new FWorkspaceItem( Name, DisplayName, Icon, bSortChildren ) );
 	}
 
 	static TSharedRef<FWorkspaceItem> NewGroup( const FText& DisplayName, const FText& TooltipText, const FSlateIcon& Icon = FSlateIcon(), const bool bSortChildren = false )
 	{
-		return MakeShareable( new FWorkspaceItem( DisplayName, TooltipText, Icon, bSortChildren ) );
+		return MakeShareable( new FWorkspaceItem( NAME_None, DisplayName, TooltipText, Icon, bSortChildren ) );
+	}
+
+	static TSharedRef<FWorkspaceItem> NewGroup( const FName& Name, const FText& DisplayName, const FText& TooltipText, const FSlateIcon& Icon = FSlateIcon(), const bool bSortChildren = false )
+	{
+		return MakeShareable( new FWorkspaceItem( Name, DisplayName, TooltipText, Icon, bSortChildren ) );
 	}
 
 	TSharedRef<FWorkspaceItem> AddGroup( const FText& InDisplayName, const FSlateIcon& InIcon = FSlateIcon(), const bool InSortChildren = false )
 	{
-		TSharedRef<FWorkspaceItem> NewItem = FWorkspaceItem::NewGroup(InDisplayName, InIcon, InSortChildren);
+		TSharedRef<FWorkspaceItem> NewItem = FWorkspaceItem::NewGroup(NAME_None, InDisplayName, InIcon, InSortChildren);
+		AddItem( NewItem );
+		return NewItem;
+	}
+
+	TSharedRef<FWorkspaceItem> AddGroup( const FName& InName, const FText& InDisplayName, const FSlateIcon& InIcon = FSlateIcon(), const bool InSortChildren = false )
+	{
+		TSharedRef<FWorkspaceItem> NewItem = FWorkspaceItem::NewGroup(InName, InDisplayName, InIcon, InSortChildren);
 		AddItem( NewItem );
 		return NewItem;
 	}
 
 	TSharedRef<FWorkspaceItem> AddGroup( const FText& InDisplayName, const FText& InTooltipText, const FSlateIcon& InIcon = FSlateIcon(), const bool InSortChildren = false )
 	{
-		TSharedRef<FWorkspaceItem> NewItem = FWorkspaceItem::NewGroup(InDisplayName, InTooltipText, InIcon, InSortChildren);
+		TSharedRef<FWorkspaceItem> NewItem = FWorkspaceItem::NewGroup(NAME_None, InDisplayName, InTooltipText, InIcon, InSortChildren);
+		AddItem( NewItem );
+		return NewItem;
+	}
+
+	TSharedRef<FWorkspaceItem> AddGroup( const FName& InName, const FText& InDisplayName, const FText& InTooltipText, const FSlateIcon& InIcon = FSlateIcon(), const bool InSortChildren = false )
+	{
+		TSharedRef<FWorkspaceItem> NewItem = FWorkspaceItem::NewGroup(InName, InDisplayName,InTooltipText, InIcon, InSortChildren);
 		AddItem( NewItem );
 		return NewItem;
 	}
 
+	const FName GetFName() const
+	{
+		return NameAttribute.Get();
+	}
+
 	const FText& GetDisplayName() const
 	{
 		return DisplayNameAttribute.Get();
@@ -149,6 +178,15 @@ PRAGMA_ENABLE_DEPRECATION_WARNINGS
 protected:
 	FWorkspaceItem( const FText& InDisplayName, const FSlateIcon& InIcon, const bool bInSortChildren )
 		: Icon(InIcon)
+		, NameAttribute(NAME_None)
+		, DisplayNameAttribute(InDisplayName)
+		, bSortChildren(bInSortChildren)
+	{
+	}
+
+	FWorkspaceItem( const FName& InName, const FText& InDisplayName, const FSlateIcon& InIcon, const bool bInSortChildren )
+		: Icon(InIcon)
+		, NameAttribute(InName)
 		, DisplayNameAttribute(InDisplayName)
 		, bSortChildren(bInSortChildren)
 	{
@@ -156,6 +194,16 @@ protected:
 
 	FWorkspaceItem( const FText& InDisplayName, const FText& InTooltipText, const FSlateIcon& InIcon, const bool bInSortChildren )
 		: Icon(InIcon)
+		, NameAttribute(NAME_None)
+		, DisplayNameAttribute(InDisplayName)
+		, TooltipTextAttribute(InTooltipText)
+		, bSortChildren(bInSortChildren)
+	{
+	}
+
+	FWorkspaceItem( const FName& InName, const FText& InDisplayName, const FText& InTooltipText, const FSlateIcon& InIcon, const bool bInSortChildren )
+		: Icon(InIcon)
+		, NameAttribute(InName)
 		, DisplayNameAttribute(InDisplayName)
 		, TooltipTextAttribute(InTooltipText)
 		, bSortChildren(bInSortChildren)
@@ -163,6 +211,7 @@ protected:
 	}
 
 	FSlateIcon Icon;
+	TAttribute<FName> NameAttribute;
 	TAttribute<FText> DisplayNameAttribute;
 	TAttribute<FText> TooltipTextAttribute;
 	UE_DEPRECATED(5.0, "Use DisplayNameAttribute instead.")
