Add FNames to workspace menu groups

This improves the Window and Tools menus of the editor. Previously, these workspace menu groups all mapped to the NAME_None section of the corresponding ToolMenus they ended up in. Because of that, their menu items were all lumped together under the same section name, which got its header from one of those groups (perhaps the first one that got registered).

diff --git a/Engine/Source/Editor/WorkspaceMenuStructure/Private/WorkspaceMenuStructureModule.cpp b/Engine/Source/Editor/WorkspaceMenuStructure/Private/WorkspaceMenuStructureModule.cpp
--- a/Engine/Source/Editor/WorkspaceMenuStructure/Private/WorkspaceMenuStructureModule.cpp
+++ b/Engine/Source/Editor/WorkspaceMenuStructure/Private/WorkspaceMenuStructureModule.cpp
@@ -141,16 +141,18 @@ public:
 
 public:
 	FWorkspaceMenuStructure()
-		: MenuRoot ( FWorkspaceItem::NewGroup(LOCTEXT( "WorkspaceMenu_Root", "Menu Root" )) )
-		, ToolsMenuRoot ( FWorkspaceItem::NewGroup(LOCTEXT( "WorkspaceToolsMenu_Root", "Tools Menu Root" )) )
-		, LevelEditorCategory ( MenuRoot->AddGroup(LOCTEXT( "WorkspaceMenu_LevelEditorCategory", "Level Editor" ), FSlateIcon(), true) )
-		, ToolsCategory ( ToolsMenuRoot->AddGroup(LOCTEXT( "WorkspaceMenu_ToolsCategory", "Tools" ), FSlateIcon(), true) )
-		, DeveloperToolsCategory ( ToolsMenuRoot->AddGroup(LOCTEXT("WorkspaceMenu_DeveloperToolsCategory", "Instrumentation"), FSlateIcon()) )
-		, DeveloperToolsLogCategory ( MenuRoot->AddGroup(LOCTEXT("WorkspaceMenu_DeveloperToolsLogCategory", "Log"), FSlateIcon(), true))
-		, EditOptions( FWorkspaceItem::NewGroup(LOCTEXT( "WorkspaceEdit_Options", "Edit Options" )) )
+		: MenuRoot(FWorkspaceItem::NewGroup(LOCTEXT("WorkspaceMenu_Root", "Menu Root")))
+		, ToolsMenuRoot(FWorkspaceItem::NewGroup(LOCTEXT("WorkspaceToolsMenu_Root", "Tools Menu Root")))
+		, LevelEditorCategory(MenuRoot->AddGroup(
+			  "LevelEditorCategory", LOCTEXT("WorkspaceMenu_LevelEditorCategory", "Level Editor"), FSlateIcon(), true))
+		, ToolsCategory(ToolsMenuRoot->AddGroup(
+			  "ToolsCategory", LOCTEXT("WorkspaceMenu_ToolsCategory", "Tools"), FSlateIcon(), true))
+		, DeveloperToolsCategory(ToolsMenuRoot->AddGroup(
+			  "DeveloperToolsCategory", LOCTEXT("WorkspaceMenu_DeveloperToolsCategory", "Instrumentation"), FSlateIcon()))
+		, DeveloperToolsLogCategory(MenuRoot->AddGroup(
+			  "DeveloperToolsLogCategory", LOCTEXT("WorkspaceMenu_DeveloperToolsLogCategory", "Log"), FSlateIcon(), true))
+		, EditOptions(FWorkspaceItem::NewGroup(LOCTEXT("WorkspaceEdit_Options", "Edit Options")))
 	{
-
-
 		ResetLevelEditorCategory();
 		ResetToolsCategory();
 	}
