Update workspace menu ToolMenu section ids to match the extension point names that were previously generated from display names

diff --git a/Engine/Source/Editor/WorkspaceMenuStructure/Private/WorkspaceMenuStructureModule.cpp b/Engine/Source/Editor/WorkspaceMenuStructure/Private/WorkspaceMenuStructureModule.cpp
--- a/Engine/Source/Editor/WorkspaceMenuStructure/Private/WorkspaceMenuStructureModule.cpp
+++ b/Engine/Source/Editor/WorkspaceMenuStructure/Private/WorkspaceMenuStructureModule.cpp
@@ -115,28 +115,60 @@ public:
 	void ResetLevelEditorCategory()
 	{
 		LevelEditorCategory->ClearItems();
-		LevelEditorViewportsCategory = LevelEditorCategory->AddGroup("LevelEditorViewport", LOCTEXT("WorkspaceMenu_LevelEditorViewportCategory", "Viewports"), LOCTEXT( "WorkspaceMenu_LevelEditorViewportCategoryTooltip", "Open a Viewport tab." ), FSlateIcon(FAppStyle::Get().GetStyleSetName(), "LevelEditor.Tabs.Viewports"), true);
-		LevelEditorDetailsCategory = LevelEditorCategory->AddGroup("LevelEditorDetail", LOCTEXT("WorkspaceMenu_LevelEditorDetailCategory", "Details" ), LOCTEXT("WorkspaceMenu_LevelEditorDetailCategoryTooltip", "Open a Details tab." ), FSlateIcon(FAppStyle::Get().GetStyleSetName(), "LevelEditor.Tabs.Details"), true );
-		LevelEditorCinematicsCategory = LevelEditorCategory->AddGroup("LevelEditorCinematics", LOCTEXT("WorkspaceMenu_LevelEditorCinematicsCategory", "Cinematics"), LOCTEXT("WorkspaceMenu_LevelEditorCinematicsCategoryTooltip", "Open a Cinematics tab."), FSlateIcon(FAppStyle::Get().GetStyleSetName(), "LevelEditor.Tabs.Cinematics"), true);
-		LevelEditorVirtualProductionCategory = LevelEditorCategory->AddGroup("LevelEditorVirtualProduction", LOCTEXT("WorkspaceMenu_LevelEditorVirtualProductionCategory", "Virtual Production"), LOCTEXT("WorkspaceMenu_LevelEditorVirtualProductionCategoryTooltip", "Open a Virtual Production tab."), FSlateIcon(FAppStyle::Get().GetStyleSetName(), "LevelEditor.Tabs.VirtualProduction"), true);
-		LevelEditorModesCategory = LevelEditorCategory->AddGroup("LevelEditorTools", LOCTEXT("WorkspaceMenu_LevelEditorToolsCategory", "Editor Modes" ), FSlateIcon(FAppStyle::Get().GetStyleSetName(), "LevelEditor.Tabs.EditorModes"), true );
-		LevelEditorWorldPartitionCategory = LevelEditorCategory->AddGroup("LevelEditorWorldPartition", LOCTEXT("WorkspaceMenu_LevelEditorWorldPartitionCategory", "World Partition"), LOCTEXT("WorkspaceMenu_LevelEditorWorldPartitionCategoryTooltip", "Open a World Partition tab."), FSlateIcon(FAppStyle::Get().GetStyleSetName(), "LevelEditor.Tabs.WorldPartition"), true);
-		LevelEditorOutlinerCategory = LevelEditorCategory->AddGroup("LevelEditorOutliner", LOCTEXT("WorkspaceMenu_LevelEditorOutlinerCategory", "Outliner"), LOCTEXT("WorkspaceMenu_LevelEditorOutlinerCategoryTooltip", "Open an Outliner tab."), FSlateIcon(FAppStyle::GetAppStyleSetName(), "LevelEditor.Tabs.Outliner"), true);
+		LevelEditorViewportsCategory = LevelEditorCategory->AddGroup("Viewports",
+			LOCTEXT("WorkspaceMenu_LevelEditorViewportCategory", "Viewports"),
+			LOCTEXT("WorkspaceMenu_LevelEditorViewportCategoryTooltip", "Open a Viewport tab."),
+			FSlateIcon(FAppStyle::Get().GetStyleSetName(), "LevelEditor.Tabs.Viewports"), true);
+		LevelEditorDetailsCategory = LevelEditorCategory->AddGroup("Details",
+			LOCTEXT("WorkspaceMenu_LevelEditorDetailCategory", "Details"),
+			LOCTEXT("WorkspaceMenu_LevelEditorDetailCategoryTooltip", "Open a Details tab."),
+			FSlateIcon(FAppStyle::Get().GetStyleSetName(), "LevelEditor.Tabs.Details"), true);
+		LevelEditorCinematicsCategory = LevelEditorCategory->AddGroup("Cinematics",
+			LOCTEXT("WorkspaceMenu_LevelEditorCinematicsCategory", "Cinematics"),
+			LOCTEXT("WorkspaceMenu_LevelEditorCinematicsCategoryTooltip", "Open a Cinematics tab."),
+			FSlateIcon(FAppStyle::Get().GetStyleSetName(), "LevelEditor.Tabs.Cinematics"), true);
+		LevelEditorVirtualProductionCategory = LevelEditorCategory->AddGroup("VirtualProduction",
+			LOCTEXT("WorkspaceMenu_LevelEditorVirtualProductionCategory", "Virtual Production"),
+			LOCTEXT("WorkspaceMenu_LevelEditorVirtualProductionCategoryTooltip", "Open a Virtual Production tab."),
+			FSlateIcon(FAppStyle::Get().GetStyleSetName(), "LevelEditor.Tabs.VirtualProduction"), true);
+		LevelEditorModesCategory = LevelEditorCategory->AddGroup("EditorModes",
+			LOCTEXT("WorkspaceMenu_LevelEditorToolsCategory", "Editor Modes"),
+			FSlateIcon(FAppStyle::Get().GetStyleSetName(), "LevelEditor.Tabs.EditorModes"), true);
+		LevelEditorWorldPartitionCategory = LevelEditorCategory->AddGroup("WorldPartition",
+			LOCTEXT("WorkspaceMenu_LevelEditorWorldPartitionCategory", "World Partition"),
+			LOCTEXT("WorkspaceMenu_LevelEditorWorldPartitionCategoryTooltip", "Open a World Partition tab."),
+			FSlateIcon(FAppStyle::Get().GetStyleSetName(), "LevelEditor.Tabs.WorldPartition"), true);
+		LevelEditorOutlinerCategory = LevelEditorCategory->AddGroup("Outliner",
+			LOCTEXT("WorkspaceMenu_LevelEditorOutlinerCategory", "Outliner"),
+			LOCTEXT("WorkspaceMenu_LevelEditorOutlinerCategoryTooltip", "Open an Outliner tab."),
+			FSlateIcon(FAppStyle::GetAppStyleSetName(), "LevelEditor.Tabs.Outliner"), true);
 	}
 
 	void ResetToolsCategory()
 	{
 		// Developer tools sections
 		DeveloperToolsCategory->ClearItems();
-		DeveloperToolsSubMenuDebugCategory =     DeveloperToolsCategory->AddGroup("DeveloperToolsDebug", LOCTEXT("WorkspaceMenu_DeveloperToolsDebugCategory", "Debug"), FSlateIcon(FAppStyle::Get().GetStyleSetName(), "LevelEditor.Tabs.Debug"), true);
-
-		DeveloperToolsDebugCategory = DeveloperToolsSubMenuDebugCategory->AddGroup("DebugTools", LOCTEXT("WorkspaceMenu_DeveloperToolsDebugCategory", "Debug"), FSlateIcon(), true);
-		AutomationToolsCategory = DeveloperToolsSubMenuDebugCategory->AddGroup("TestingTools", LOCTEXT("WorkspaceMenu_AutomationToolsCategory", "Testing"), FSlateIcon(), true);
-
-		DeveloperToolsProfilingCategory = DeveloperToolsCategory->AddGroup("DeveloperToolsProfiling", LOCTEXT("WorkspaceMenu_DeveloperToolsProfilingCategory", "Profile"),  FSlateIcon(FAppStyle::Get().GetStyleSetName(), "LevelEditor.Profile"), true);
-		DeveloperToolsAuditCategory =     DeveloperToolsCategory->AddGroup("DeveloperToolsAudit", LOCTEXT("WorkspaceMenu_DeveloperToolsAuditCategory", "Audit"), FSlateIcon(FAppStyle::Get().GetStyleSetName(), "LevelEditor.Audit"), true);
-		DeveloperToolsPlatformsCategory = DeveloperToolsCategory->AddGroup("DeveloperToolsPlatforms", LOCTEXT("WorkspaceMenu_DeveloperToolsPlatforms:WCategory", "Platforms"), FSlateIcon(FAppStyle::Get().GetStyleSetName(), "LevelEditor.Platforms"), true);
-		DeveloperToolsMiscCategory =      DeveloperToolsCategory->AddGroup("DeveloperToolsMisc", LOCTEXT("WorkspaceMenu_DeveloperToolsMiscCategory", "Miscellaneous"), FSlateIcon(FAppStyle::Get().GetStyleSetName(), "LevelEditor.Tabs.Debug"), true);
+		DeveloperToolsSubMenuDebugCategory = DeveloperToolsCategory->AddGroup("Debug",
+			LOCTEXT("WorkspaceMenu_DeveloperToolsDebugCategory", "Debug"),
+			FSlateIcon(FAppStyle::Get().GetStyleSetName(), "LevelEditor.Tabs.Debug"), true);
+
+		DeveloperToolsDebugCategory = DeveloperToolsSubMenuDebugCategory->AddGroup(
+			"Debug", LOCTEXT("WorkspaceMenu_DeveloperToolsDebugCategory", "Debug"), FSlateIcon(), true);
+		AutomationToolsCategory = DeveloperToolsSubMenuDebugCategory->AddGroup(
+			"Testing", LOCTEXT("WorkspaceMenu_AutomationToolsCategory", "Testing"), FSlateIcon(), true);
+
+		DeveloperToolsProfilingCategory = DeveloperToolsCategory->AddGroup("Profile",
+			LOCTEXT("WorkspaceMenu_DeveloperToolsProfilingCategory", "Profile"),
+			FSlateIcon(FAppStyle::Get().GetStyleSetName(), "LevelEditor.Profile"), true);
+		DeveloperToolsAuditCategory = DeveloperToolsCategory->AddGroup("Audit",
+			LOCTEXT("WorkspaceMenu_DeveloperToolsAuditCategory", "Audit"),
+			FSlateIcon(FAppStyle::Get().GetStyleSetName(), "LevelEditor.Audit"), true);
+		DeveloperToolsPlatformsCategory = DeveloperToolsCategory->AddGroup("Platforms",
+			LOCTEXT("WorkspaceMenu_DeveloperToolsPlatforms:WCategory", "Platforms"),
+			FSlateIcon(FAppStyle::Get().GetStyleSetName(), "LevelEditor.Platforms"), true);
+		DeveloperToolsMiscCategory = DeveloperToolsCategory->AddGroup("Miscellaneous",
+			LOCTEXT("WorkspaceMenu_DeveloperToolsMiscCategory", "Miscellaneous"),
+			FSlateIcon(FAppStyle::Get().GetStyleSetName(), "LevelEditor.Tabs.Debug"), true);
 	}
 
 public:
@@ -144,13 +176,13 @@ public:
 		: MenuRoot(FWorkspaceItem::NewGroup(LOCTEXT("WorkspaceMenu_Root", "Menu Root")))
 		, ToolsMenuRoot(FWorkspaceItem::NewGroup(LOCTEXT("WorkspaceToolsMenu_Root", "Tools Menu Root")))
 		, LevelEditorCategory(MenuRoot->AddGroup(
-			  "LevelEditorCategory", LOCTEXT("WorkspaceMenu_LevelEditorCategory", "Level Editor"), FSlateIcon(), true))
-		, ToolsCategory(ToolsMenuRoot->AddGroup(
-			  "ToolsCategory", LOCTEXT("WorkspaceMenu_ToolsCategory", "Tools"), FSlateIcon(), true))
+			  "LevelEditor", LOCTEXT("WorkspaceMenu_LevelEditorCategory", "Level Editor"), FSlateIcon(), true))
+		, ToolsCategory(
+			  ToolsMenuRoot->AddGroup("Tools", LOCTEXT("WorkspaceMenu_ToolsCategory", "Tools"), FSlateIcon(), true))
 		, DeveloperToolsCategory(ToolsMenuRoot->AddGroup(
-			  "DeveloperToolsCategory", LOCTEXT("WorkspaceMenu_DeveloperToolsCategory", "Instrumentation"), FSlateIcon()))
-		, DeveloperToolsLogCategory(MenuRoot->AddGroup(
-			  "DeveloperToolsLogCategory", LOCTEXT("WorkspaceMenu_DeveloperToolsLogCategory", "Log"), FSlateIcon(), true))
+			  "Instrumentation", LOCTEXT("WorkspaceMenu_DeveloperToolsCategory", "Instrumentation"), FSlateIcon()))
+		, DeveloperToolsLogCategory(
+			  MenuRoot->AddGroup("Log", LOCTEXT("WorkspaceMenu_DeveloperToolsLogCategory", "Log"), FSlateIcon(), true))
 		, EditOptions(FWorkspaceItem::NewGroup(LOCTEXT("WorkspaceEdit_Options", "Edit Options")))
 	{
 		ResetLevelEditorCategory();
