WorkspaceMenuStructure: Remove "Tools > Miscellaneous" icon

diff --git a/Engine/Source/Editor/WorkspaceMenuStructure/Private/WorkspaceMenuStructureModule.cpp b/Engine/Source/Editor/WorkspaceMenuStructure/Private/WorkspaceMenuStructureModule.cpp
--- a/Engine/Source/Editor/WorkspaceMenuStructure/Private/WorkspaceMenuStructureModule.cpp
+++ b/Engine/Source/Editor/WorkspaceMenuStructure/Private/WorkspaceMenuStructureModule.cpp
@@ -167,8 +167,7 @@ public:
 			LOCTEXT("WorkspaceMenu_DeveloperToolsPlatforms:WCategory", "Platforms"),
 			FSlateIcon(FAppStyle::Get().GetStyleSetName(), "LevelEditor.Platforms"), true);
 		DeveloperToolsMiscCategory = DeveloperToolsCategory->AddGroup("Miscellaneous",
-			LOCTEXT("WorkspaceMenu_DeveloperToolsMiscCategory", "Miscellaneous"),
-			FSlateIcon(FAppStyle::Get().GetStyleSetName(), "LevelEditor.Tabs.Debug"), true);
+			LOCTEXT("WorkspaceMenu_DeveloperToolsMiscCategory", "Miscellaneous"), FSlateIcon(), true);
 	}
 
 public:
