Fix bug where main menu icons did not appear

The issue was that the editor style loaded after the workspace menu structure, so the workspace menu categories fetched their icons too early. This resulted in a few submenus under Window and Tools binding their icons to "CoreStyle" instead of "EditorStyle" where they really live. The fix is to explicitly load EditorStyle before it is needed by WorkspaceMenuStructure so we know its loaded.

I had to remove TurnkeySupport's dependency on WorkspaceMenuStructureModule to avoid one of TurnkeySupport non-Editor dependent modules from indirectly depending on the Editor module WorkspaceMenuStructureModule.

diff --git a/Engine/Source/Editor/TurnkeySupport/TurnkeySupport.Build.cs b/Engine/Source/Editor/TurnkeySupport/TurnkeySupport.Build.cs
--- a/Engine/Source/Editor/TurnkeySupport/TurnkeySupport.Build.cs
+++ b/Engine/Source/Editor/TurnkeySupport/TurnkeySupport.Build.cs
@@ -20,7 +20,6 @@ public class TurnkeySupport : ModuleRules
 				"SlateCore",
 				"TargetPlatform",
 				"DesktopPlatform",
-				"WorkspaceMenuStructure",
 				"MessageLog",
  				"Projects",
  				"ToolMenus",
diff --git a/Engine/Source/Editor/WorkspaceMenuStructure/Private/WorkspaceMenuStructureModule.cpp b/Engine/Source/Editor/WorkspaceMenuStructure/Private/WorkspaceMenuStructureModule.cpp
--- a/Engine/Source/Editor/WorkspaceMenuStructure/Private/WorkspaceMenuStructureModule.cpp
+++ b/Engine/Source/Editor/WorkspaceMenuStructure/Private/WorkspaceMenuStructureModule.cpp
@@ -3,6 +3,7 @@
 #include "WorkspaceMenuStructureModule.h"
 
 #include "Framework/Docking/WorkspaceItem.h"
+#include "Interfaces/IEditorStyleModule.h"
 #include "Internationalization/Internationalization.h"
 #include "Misc/AssertionMacros.h"
 #include "Styling/AppStyle.h"
@@ -220,6 +221,9 @@ private:
 
 void FWorkspaceMenuStructureModule::StartupModule()
 {
+	// The workspace menu categories rely on icons from the editor style.
+	FModuleManager::LoadModuleChecked<IEditorStyleModule>("EditorStyle");
+
 	WorkspaceMenuStructure = MakeShareable(new FWorkspaceMenuStructure);
 }
 
diff --git a/Engine/Source/Editor/WorkspaceMenuStructure/WorkspaceMenuStructure.Build.cs b/Engine/Source/Editor/WorkspaceMenuStructure/WorkspaceMenuStructure.Build.cs
--- a/Engine/Source/Editor/WorkspaceMenuStructure/WorkspaceMenuStructure.Build.cs
+++ b/Engine/Source/Editor/WorkspaceMenuStructure/WorkspaceMenuStructure.Build.cs
@@ -9,6 +9,7 @@ public class WorkspaceMenuStructure : ModuleRules
 		PrivateDependencyModuleNames.AddRange(
 			new string[] {
 				"Core",
+				"EditorStyle",
 				"SlateCore",
 				"Slate"
 			}
