Fix build issue caused by CL 37965333

diff --git a/Engine/Source/Editor/WorkspaceMenuStructure/Private/WorkspaceMenuStructureModule.cpp b/Engine/Source/Editor/WorkspaceMenuStructure/Private/WorkspaceMenuStructureModule.cpp
--- a/Engine/Source/Editor/WorkspaceMenuStructure/Private/WorkspaceMenuStructureModule.cpp
+++ b/Engine/Source/Editor/WorkspaceMenuStructure/Private/WorkspaceMenuStructureModule.cpp
@@ -3,7 +3,9 @@
 #include "WorkspaceMenuStructureModule.h"
 
 #include "Framework/Docking/WorkspaceItem.h"
+#if WITH_EDITOR
 #include "Interfaces/IEditorStyleModule.h"
+#endif
 #include "Internationalization/Internationalization.h"
 #include "Misc/AssertionMacros.h"
 #include "Styling/AppStyle.h"
@@ -221,8 +223,10 @@ private:
 
 void FWorkspaceMenuStructureModule::StartupModule()
 {
+#if WITH_EDITOR
 	// The workspace menu categories rely on icons from the editor style.
 	FModuleManager::LoadModuleChecked<IEditorStyleModule>("EditorStyle");
+#endif
 
 	WorkspaceMenuStructure = MakeShareable(new FWorkspaceMenuStructure);
 }
diff --git a/Engine/Source/Editor/WorkspaceMenuStructure/WorkspaceMenuStructure.Build.cs b/Engine/Source/Editor/WorkspaceMenuStructure/WorkspaceMenuStructure.Build.cs
--- a/Engine/Source/Editor/WorkspaceMenuStructure/WorkspaceMenuStructure.Build.cs
+++ b/Engine/Source/Editor/WorkspaceMenuStructure/WorkspaceMenuStructure.Build.cs
@@ -9,10 +9,14 @@ public class WorkspaceMenuStructure : ModuleRules
 		PrivateDependencyModuleNames.AddRange(
 			new string[] {
 				"Core",
-				"EditorStyle",
 				"SlateCore",
 				"Slate"
 			}
 		);
+		
+		if (Target.bCompileAgainstEditor)
+		{
+			PublicDependencyModuleNames.Add("EditorStyle");
+		}
 	}
 }
