Apparently this editor module is used in game which makes sense.

#Synced-CL 48605542

diff --git a/Engine/Source/Editor/WorkspaceMenuStructure/WorkspaceMenuStructure.Build.cs b/Engine/Source/Editor/WorkspaceMenuStructure/WorkspaceMenuStructure.Build.cs
--- a/Engine/Source/Editor/WorkspaceMenuStructure/WorkspaceMenuStructure.Build.cs
+++ b/Engine/Source/Editor/WorkspaceMenuStructure/WorkspaceMenuStructure.Build.cs
@@ -2,7 +2,7 @@
 
 using UnrealBuildTool;
 
-[SupportedTargetTypes(TargetType.Editor, TargetType.Program)]
+[SupportedTargetTypes(TargetType.Editor, TargetType.Game, TargetType.Program)]
 public class WorkspaceMenuStructure : ModuleRules
 {
 	public WorkspaceMenuStructure(ReadOnlyTargetRules Target) : base(Target)
