WorkspaceMenuStructure is used in server too apparently.

#Synced-CL 50155317

diff --git a/Engine/Source/Editor/WorkspaceMenuStructure/WorkspaceMenuStructure.Build.cs b/Engine/Source/Editor/WorkspaceMenuStructure/WorkspaceMenuStructure.Build.cs
--- a/Engine/Source/Editor/WorkspaceMenuStructure/WorkspaceMenuStructure.Build.cs
+++ b/Engine/Source/Editor/WorkspaceMenuStructure/WorkspaceMenuStructure.Build.cs
@@ -2,7 +2,7 @@
 
 using UnrealBuildTool;
 
-[SupportedTargetTypes(TargetType.Editor, TargetType.Game, TargetType.Program)]
+[SupportedTargetTypes(TargetType.Editor, TargetType.Game, TargetType.Program, TargetType.Server)]
 public class WorkspaceMenuStructure : ModuleRules
 {
 	public WorkspaceMenuStructure(ReadOnlyTargetRules Target) : base(Target)
